#UE4 Disable CanContainContent in a few plugins that have no assets or maps so they dont appear in the content browser.

diff --git a/Engine/Plugins/Editor/BlueprintMaterialTextureNodes/BlueprintMaterialTextureNodes.uplugin b/Engine/Plugins/Editor/BlueprintMaterialTextureNodes/BlueprintMaterialTextureNodes.uplugin
--- a/Engine/Plugins/Editor/BlueprintMaterialTextureNodes/BlueprintMaterialTextureNodes.uplugin
+++ b/Engine/Plugins/Editor/BlueprintMaterialTextureNodes/BlueprintMaterialTextureNodes.uplugin
@@ -10,7 +10,7 @@
 	"DocsURL": "",
 	"MarketplaceURL": "",
 	"SupportURL": "",
-	"CanContainContent": true,
+	"CanContainContent": false,
 	"IsBetaVersion": false,
 	"Installed": false,
 	"Modules": [
diff --git a/Engine/Plugins/Editor/CryptoKeys/CryptoKeys.uplugin b/Engine/Plugins/Editor/CryptoKeys/CryptoKeys.uplugin
--- a/Engine/Plugins/Editor/CryptoKeys/CryptoKeys.uplugin
+++ b/Engine/Plugins/Editor/CryptoKeys/CryptoKeys.uplugin
@@ -10,7 +10,7 @@
 	"DocsURL": "",
 	"MarketplaceURL": "",
 	"SupportURL": "",
-	"CanContainContent": true,
+	"CanContainContent": false,
 	"IsBetaVersion": false,
 	"Installed": false,
 	"EnabledByDefault": true,
diff --git a/Engine/Plugins/Editor/CurveEditorTools/CurveEditorTools.uplugin b/Engine/Plugins/Editor/CurveEditorTools/CurveEditorTools.uplugin
--- a/Engine/Plugins/Editor/CurveEditorTools/CurveEditorTools.uplugin
+++ b/Engine/Plugins/Editor/CurveEditorTools/CurveEditorTools.uplugin
@@ -10,7 +10,7 @@
 	"DocsURL": "",
 	"MarketplaceURL": "",
 	"SupportURL": "",
-	"CanContainContent": true,
+	"CanContainContent": false,
 	"IsBetaVersion": false,
 	"Installed": false,
 	"EnabledByDefault": true,
diff --git a/Engine/Plugins/Experimental/AutomationUtils/AutomationUtils.uplugin b/Engine/Plugins/Experimental/AutomationUtils/AutomationUtils.uplugin
--- a/Engine/Plugins/Experimental/AutomationUtils/AutomationUtils.uplugin
+++ b/Engine/Plugins/Experimental/AutomationUtils/AutomationUtils.uplugin
@@ -10,7 +10,7 @@
 	"MarketplaceURL": "",
 	"SupportURL": "",
 	"EnabledByDefault": true,
-	"CanContainContent": true,
+	"CanContainContent": false,
 	"IsBetaVersion": true,
 	"Installed": false,
 	"Modules": [
diff --git a/Engine/Plugins/Experimental/PythonScriptPlugin/PythonScriptPlugin.uplugin b/Engine/Plugins/Experimental/PythonScriptPlugin/PythonScriptPlugin.uplugin
--- a/Engine/Plugins/Experimental/PythonScriptPlugin/PythonScriptPlugin.uplugin
+++ b/Engine/Plugins/Experimental/PythonScriptPlugin/PythonScriptPlugin.uplugin
@@ -11,7 +11,7 @@
 	"MarketplaceURL" : "",
 	"SupportURL" : "https://forums.unrealengine.com/",
 	"EnabledByDefault" : false,
-	"CanContainContent" : true,
+	"CanContainContent" : false,
 	"IsBetaVersion" : true,
 	"Installed" : false,
 	
diff --git a/Engine/Plugins/MovieScene/SequencerScripting/SequencerScripting.uplugin b/Engine/Plugins/MovieScene/SequencerScripting/SequencerScripting.uplugin
--- a/Engine/Plugins/MovieScene/SequencerScripting/SequencerScripting.uplugin
+++ b/Engine/Plugins/MovieScene/SequencerScripting/SequencerScripting.uplugin
@@ -10,7 +10,7 @@
 	"DocsURL": "",
 	"MarketplaceURL": "",
 	"SupportURL": "",
-	"CanContainContent": true,
+	"CanContainContent": false,
 	"IsBetaVersion": true,
 	"Installed": false,
 	"Plugins": [
diff --git a/Engine/Plugins/Runtime/Synthesis/Synthesis.uplugin b/Engine/Plugins/Runtime/Synthesis/Synthesis.uplugin
--- a/Engine/Plugins/Runtime/Synthesis/Synthesis.uplugin
+++ b/Engine/Plugins/Runtime/Synthesis/Synthesis.uplugin
@@ -11,7 +11,7 @@
 	"MarketplaceURL" : "",
 	"SupportURL" : "",
 	"EnabledByDefault" : false,
-	"CanContainContent" : true,
+	"CanContainContent" : false,
 	"IsBetaVersion" : true,
 	"Installed" : false,
 	"Modules" :
diff --git a/Engine/Plugins/Runtime/TimeSynth/TimeSynth.uplugin b/Engine/Plugins/Runtime/TimeSynth/TimeSynth.uplugin
--- a/Engine/Plugins/Runtime/TimeSynth/TimeSynth.uplugin
+++ b/Engine/Plugins/Runtime/TimeSynth/TimeSynth.uplugin
@@ -11,7 +11,7 @@
 	"MarketplaceURL" : "",
 	"SupportURL" : "",
 	"EnabledByDefault" : false,
-	"CanContainContent" : true,
+	"CanContainContent" : false,
 	"IsBetaVersion" : true,
 	"Installed" : false,
 	"Modules" :
