Fixing bad include paths.

diff --git a/Engine/Plugins/Compositing/Composure/Source/ComposureLayersEditor/Private/ComposureDetailCustomizations.cpp b/Engine/Plugins/Compositing/Composure/Source/ComposureLayersEditor/Private/ComposureDetailCustomizations.cpp
--- a/Engine/Plugins/Compositing/Composure/Source/ComposureLayersEditor/Private/ComposureDetailCustomizations.cpp
+++ b/Engine/Plugins/Compositing/Composure/Source/ComposureLayersEditor/Private/ComposureDetailCustomizations.cpp
@@ -10,7 +10,7 @@
 #include "Customizations/ColorStructCustomization.h"
 #include "ICompElementManager.h"
 #include "Widgets/Input/SButton.h"
-#include "MaterialInterface.h"
+#include "Materials/MaterialInterface.h"
 #include "Widgets/SToolTip.h"
 #include "ComposureEditorStyle.h"
 #include "Widgets/SCompElementPickerWindow.h"
diff --git a/Engine/Plugins/Compositing/Composure/Source/ComposureLayersEditor/Private/Widgets/SCompElementPreviewPane.cpp b/Engine/Plugins/Compositing/Composure/Source/ComposureLayersEditor/Private/Widgets/SCompElementPreviewPane.cpp
--- a/Engine/Plugins/Compositing/Composure/Source/ComposureLayersEditor/Private/Widgets/SCompElementPreviewPane.cpp
+++ b/Engine/Plugins/Compositing/Composure/Source/ComposureLayersEditor/Private/Widgets/SCompElementPreviewPane.cpp
@@ -3,7 +3,7 @@
 #include "Widgets/SCompElementPreviewPane.h"
 #include "Framework/Commands/UICommandList.h"
 #include "Widgets/Input/SButton.h"
-#include "MaterialInterface.h"
+#include "Materials/MaterialInterface.h"
 #include "Widgets/Layout/SScaleBox.h"
 #include "Framework/MultiBox/MultiBoxBuilder.h"
 #include "SlateMaterialBrush.h"
diff --git a/Engine/Plugins/Editor/BlueprintMaterialTextureNodes/Source/BlueprintMaterialTextureNodes/Private/BlueprintMaterialTextureNodesBPLibrary.cpp b/Engine/Plugins/Editor/BlueprintMaterialTextureNodes/Source/BlueprintMaterialTextureNodes/Private/BlueprintMaterialTextureNodesBPLibrary.cpp
--- a/Engine/Plugins/Editor/BlueprintMaterialTextureNodes/Source/BlueprintMaterialTextureNodes/Private/BlueprintMaterialTextureNodesBPLibrary.cpp
+++ b/Engine/Plugins/Editor/BlueprintMaterialTextureNodes/Source/BlueprintMaterialTextureNodes/Private/BlueprintMaterialTextureNodesBPLibrary.cpp
@@ -8,7 +8,7 @@
 //includes for asset creation
 #include "AssetToolsModule.h"
 #include "IAssetTools.h"
-#include "MaterialInstanceBasePropertyOverrides.h"
+#include "Materials/MaterialInstanceBasePropertyOverrides.h"
 #include "PackageTools.h"
 #include "Editor.h"
 #include "Factories/MaterialInstanceConstantFactoryNew.h"
@@ -16,7 +16,7 @@
 
 //Material and texture includes
 #include "Engine/TextureRenderTarget2D.h"
-#include "MaterialInterface.h"
+#include "Materials/MaterialInterface.h"
 #include "Materials/MaterialInstanceConstant.h"
 #include "MaterialShared.h"
 #include "ImageCoreUtils.h"
