Cleaning up a number of direct module includes.

diff --git a/Engine/Plugins/Editor/BlueprintMaterialTextureNodes/Source/BlueprintMaterialTextureNodes/BlueprintMaterialTextureNodes.Build.cs b/Engine/Plugins/Editor/BlueprintMaterialTextureNodes/Source/BlueprintMaterialTextureNodes/BlueprintMaterialTextureNodes.Build.cs
--- a/Engine/Plugins/Editor/BlueprintMaterialTextureNodes/Source/BlueprintMaterialTextureNodes/BlueprintMaterialTextureNodes.Build.cs
+++ b/Engine/Plugins/Editor/BlueprintMaterialTextureNodes/Source/BlueprintMaterialTextureNodes/BlueprintMaterialTextureNodes.Build.cs
@@ -1,5 +1,6 @@
 // Copyright Epic Games, Inc. All Rights Reserved.
 
+using System.IO;
 using UnrealBuildTool;
 
 public class BlueprintMaterialTextureNodes : ModuleRules
@@ -7,60 +8,19 @@ public class BlueprintMaterialTextureNodes : ModuleRules
 	public BlueprintMaterialTextureNodes(ReadOnlyTargetRules Target) : base(Target)
 	{
 		PCHUsage = ModuleRules.PCHUsageMode.UseExplicitOrSharedPCHs;
-		
-		PublicIncludePaths.AddRange(
-			new string[] {
-                "Runtime/Engine/Classes/Materials",
-				
-				// ... add public include paths required here ...
-			}
-			);			
-		
-		PublicDependencyModuleNames.AddRange(
-			new string[]
-			{
-				"Core",
-				"ImageCore",
-				"AssetRegistry",
-                "AssetTools",
-                "RHI",
-                "StaticMeshEditor",
-				
-				// ... add other public dependencies that you statically link with here ...
-			}
-			);
-			
-		
-		PrivateDependencyModuleNames.AddRange(
-			new string[]
-			{
-                "CoreUObject",
-				"EditorFramework",
-                "UnrealEd",
-                "Engine",
-                "PropertyEditor",
-                "Slate",
-                "SlateCore",
 
-				// ... add private dependencies that you statically link with here ...	
-			}
-			);
-
-        PrivateIncludePathModuleNames.AddRange(
-           new string[]
-           {
-                "Settings",
-                "AssetTools",
-                "AssetRegistry",
-               //"StaticMeshEditor",
-            }
-            );
+		PublicDependencyModuleNames.AddRange(new string[]
+			{
+				"Engine",
+			});
 
-        DynamicallyLoadedModuleNames.AddRange(
-			new string[]
+		PrivateDependencyModuleNames.AddRange(new string[]
 			{
-				// ... add any modules that your module loads dynamically here ...
-			}
-			);
+				"Core",
+				"CoreUObject",
+				"ImageCore",
+				"RHI",
+				"UnrealEd",
+			});
 	}
 }
diff --git a/Engine/Plugins/FX/NiagaraFluids/Source/NiagaraFluids/NiagaraFluids.Build.cs b/Engine/Plugins/FX/NiagaraFluids/Source/NiagaraFluids/NiagaraFluids.Build.cs
--- a/Engine/Plugins/FX/NiagaraFluids/Source/NiagaraFluids/NiagaraFluids.Build.cs
+++ b/Engine/Plugins/FX/NiagaraFluids/Source/NiagaraFluids/NiagaraFluids.Build.cs
@@ -6,47 +6,9 @@ namespace UnrealBuildTool.Rules
 	{
 		public NiagaraFluids(ReadOnlyTargetRules Target) : base(Target)
 		{
-			PublicIncludePaths.AddRange(
-				new string[] {
-					// ... add public include paths required here ...
-				}
-				);
-
-			PrivateIncludePaths.AddRange(
-				new string[] {
-                    "../../../../Source/Runtime/Engine/",
-                    "../../../../Source/Runtime/RHI/",
-					// ... add other private include paths required here ...
-				}
-				);
-            
-            PublicDependencyModuleNames.AddRange(
-				new string[]
-				{
-                    "Core",
-                    "CoreUObject",
-                    "Engine",
-                    "RHI",
-                    "Engine",
-                    "RenderCore",
-					// ... add other public dependencies that you statically link with here ...
-				}
-				);
-
-			PrivateDependencyModuleNames.AddRange(
-				new string[]
-                {
-                    "Projects",
-					// ... add private dependencies that you statically link with here ...
-				}
-				);
-
-			DynamicallyLoadedModuleNames.AddRange(
-				new string[]
-				{
-					// ... add any modules that your module loads dynamically here ...
-				}
-				);
+			PublicDependencyModuleNames.Add("Core");
+			PrivateDependencyModuleNames.Add("RenderCore");
+			PrivateDependencyModuleNames.Add("Projects");
 		}
 	}
 }
diff --git a/Engine/Plugins/Runtime/Oculus/OculusVR/Source/OculusHMD/OculusHMD.Build.cs b/Engine/Plugins/Runtime/Oculus/OculusVR/Source/OculusHMD/OculusHMD.Build.cs
--- a/Engine/Plugins/Runtime/Oculus/OculusVR/Source/OculusHMD/OculusHMD.Build.cs
+++ b/Engine/Plugins/Runtime/Oculus/OculusVR/Source/OculusHMD/OculusHMD.Build.cs
@@ -9,13 +9,6 @@ namespace UnrealBuildTool.Rules
 	{
 		public OculusHMD(ReadOnlyTargetRules Target) : base(Target)
 		{
-			PrivateIncludePaths.AddRange(
-				new string[] {
-					// Relative to Engine\Plugins\Runtime\Oculus\OculusVR\Source
-					"../../../../../Source/Runtime/Engine/Classes/Components",
-					"../../../../../Source/Runtime/Engine/Classes/Kismet",
-				});
-
 			PublicIncludePathModuleNames.AddRange(
 				new string[] {
 					"Launch",
diff --git a/Engine/Plugins/Runtime/Oculus/OculusVR/Source/OculusInput/OculusInput.Build.cs b/Engine/Plugins/Runtime/Oculus/OculusVR/Source/OculusInput/OculusInput.Build.cs
--- a/Engine/Plugins/Runtime/Oculus/OculusVR/Source/OculusInput/OculusInput.Build.cs
+++ b/Engine/Plugins/Runtime/Oculus/OculusVR/Source/OculusInput/OculusInput.Build.cs
@@ -1,5 +1,7 @@
 // Copyright Epic Games, Inc. All Rights Reserved.
 
+using System.IO;
+
 namespace UnrealBuildTool.Rules
 {
 	public class OculusInput : ModuleRules
@@ -32,9 +34,7 @@ namespace UnrealBuildTool.Rules
 
 			PrivateIncludePaths.AddRange(
 				new string[] {
-					// Relative to Engine\Plugins\Runtime\Oculus\OculusVR\Source
-					"OculusHMD/Private",
-					"../../../../../Source/Runtime/Engine/Classes/Components",
+					Path.Combine(GetModuleDirectory("OculusHMD"), "Private"),
 				});
 
 			if (Target.Platform == UnrealTargetPlatform.Win64)
diff --git a/Engine/Plugins/Runtime/Oculus/OculusVR/Source/OculusMR/OculusMR.Build.cs b/Engine/Plugins/Runtime/Oculus/OculusVR/Source/OculusMR/OculusMR.Build.cs
--- a/Engine/Plugins/Runtime/Oculus/OculusVR/Source/OculusMR/OculusMR.Build.cs
+++ b/Engine/Plugins/Runtime/Oculus/OculusVR/Source/OculusMR/OculusMR.Build.cs
@@ -1,5 +1,7 @@
 // Copyright Epic Games, Inc. All Rights Reserved.
 
+using System.IO;
+
 namespace UnrealBuildTool.Rules
 {
 	public class OculusMR : ModuleRules
@@ -35,11 +37,9 @@ namespace UnrealBuildTool.Rules
 
 			PrivateIncludePaths.AddRange(
 				new string[] {
-					// Relative to Engine\Plugins\Runtime\Oculus\OculusVR\Source
-					"OculusHMD/Private",
-                    "OculusInput/Private",
-					"../../../../../Source/Runtime/Engine/Classes/Components",
-					System.IO.Path.Combine(GetModuleDirectory("MediaAssets"), "Private"),
+					Path.Combine(GetModuleDirectory("OculusHMD"), "Private"),
+					Path.Combine(GetModuleDirectory("OculusInput"), "Private"),
+					Path.Combine(GetModuleDirectory("MediaAssets"), "Private"),
                 });
 
             if (Target.Platform == UnrealTargetPlatform.Win64)
diff --git a/Engine/Plugins/Runtime/Oculus/OculusVR/Source/OculusOpenXRHMD/OculusOpenXRHMD.Build.cs b/Engine/Plugins/Runtime/Oculus/OculusVR/Source/OculusOpenXRHMD/OculusOpenXRHMD.Build.cs
--- a/Engine/Plugins/Runtime/Oculus/OculusVR/Source/OculusOpenXRHMD/OculusOpenXRHMD.Build.cs
+++ b/Engine/Plugins/Runtime/Oculus/OculusVR/Source/OculusOpenXRHMD/OculusOpenXRHMD.Build.cs
@@ -11,10 +11,7 @@ namespace UnrealBuildTool.Rules
         {
             PrivateIncludePaths.AddRange(
                 new string[] {
-                    // Relative to Engine\Plugins\Runtime\Oculus\OculusOpenXR\Source
-                    "../../../OpenXR/Source/OpenXRHMD/Private",
-					"../../../../../Source/Runtime/Engine/Classes/Components",
-                    "../../../../../Source/Runtime/Engine/Classes/Kismet",
+					Path.Combine(GetModuleDirectory("OpenXRHMD"), "Private"),
                 });
 
             PublicIncludePathModuleNames.AddRange(
diff --git a/Engine/Plugins/Tests/RHITests/Source/RHITests/RHITests.Build.cs b/Engine/Plugins/Tests/RHITests/Source/RHITests/RHITests.Build.cs
--- a/Engine/Plugins/Tests/RHITests/Source/RHITests/RHITests.Build.cs
+++ b/Engine/Plugins/Tests/RHITests/Source/RHITests/RHITests.Build.cs
@@ -6,14 +6,6 @@ public class RHITests : ModuleRules
 {
 	public RHITests(ReadOnlyTargetRules Target) : base(Target)
 	{				
-		PrivateIncludePaths.AddRange(
-			new string[] {
-				"../../../../Source/Runtime/Engine/",
-				"../../../../Source/Runtime/RHI/",
-				// ... add other private include paths required here ...
-			}
-		);
-			
 		PublicDependencyModuleNames.AddRange(
 			new string[]
 			{
@@ -23,7 +15,6 @@ public class RHITests : ModuleRules
                 "RHI",
                 "Engine",
                 "RenderCore",
-				// ... add other public dependencies that you statically link with here ...
 			}
 		);
 			
@@ -32,7 +23,6 @@ public class RHITests : ModuleRules
 			new string[]
 			{
 				"Projects",
-				// ... add private dependencies that you statically link with here ...	
 			}
 		);
 	}
diff --git a/Engine/Source/Developer/LogVisualizer/LogVisualizer.Build.cs b/Engine/Source/Developer/LogVisualizer/LogVisualizer.Build.cs
--- a/Engine/Source/Developer/LogVisualizer/LogVisualizer.Build.cs
+++ b/Engine/Source/Developer/LogVisualizer/LogVisualizer.Build.cs
@@ -1,5 +1,6 @@
 // Copyright Epic Games, Inc. All Rights Reserved.
 
+using System.IO;
 using UnrealBuildTool;
 
 public class LogVisualizer : ModuleRules
@@ -14,12 +15,7 @@ public class LogVisualizer : ModuleRules
 			}
 		);
 
-		PublicIncludePaths.AddRange(
-			new string[] {
-				"Runtime/Engine/Classes",
-				"Editor/WorkspaceMenuStructure/Public"
-			}
-		);
+		PublicIncludePathModuleNames.Add("WorkspaceMenuStructure");
 
 		PrivateDependencyModuleNames.AddRange(
 			new string[] {
diff --git a/Engine/Source/Runtime/AIModule/AIModule.Build.cs b/Engine/Source/Runtime/AIModule/AIModule.Build.cs
--- a/Engine/Source/Runtime/AIModule/AIModule.Build.cs
+++ b/Engine/Source/Runtime/AIModule/AIModule.Build.cs
@@ -1,5 +1,7 @@
 // Copyright Epic Games, Inc. All Rights Reserved.
 
+using System.IO;
+
 namespace UnrealBuildTool.Rules
 {
     public class AIModule : ModuleRules
@@ -15,7 +17,7 @@ namespace UnrealBuildTool.Rules
 
             PrivateIncludePaths.AddRange(
                 new string[] {
-                    "Runtime/Engine/Private",
+					Path.Combine(GetModuleDirectory("Engine"), "Private"),
                 }
                 );
 
diff --git a/Engine/Source/Runtime/Android/AndroidLocalNotification/AndroidLocalNotification.Build.cs b/Engine/Source/Runtime/Android/AndroidLocalNotification/AndroidLocalNotification.Build.cs
--- a/Engine/Source/Runtime/Android/AndroidLocalNotification/AndroidLocalNotification.Build.cs
+++ b/Engine/Source/Runtime/Android/AndroidLocalNotification/AndroidLocalNotification.Build.cs
@@ -8,11 +8,11 @@ public class AndroidLocalNotification : ModuleRules
 	{
 		BinariesSubFolder = "Android";
 
-        PublicIncludePaths.AddRange(new string[]
-        {
-            "Runtime/Android/AndroidLocalNotification/Public",
-            "Runtime/Engine/Public",
-        });
+		PublicIncludePathModuleNames.AddRange(new string[]
+		{
+			"AndroidLocalNotification",
+			"Engine",
+		});
 
         PrivateDependencyModuleNames.AddRange(new string[]
 		{
diff --git a/Engine/Source/Runtime/Core/Core.Build.cs b/Engine/Source/Runtime/Core/Core.Build.cs
--- a/Engine/Source/Runtime/Core/Core.Build.cs
+++ b/Engine/Source/Runtime/Core/Core.Build.cs
@@ -32,20 +32,15 @@ public class Core : ModuleRules
 		PublicDependencyModuleNames.Add("TraceLog");
 		PublicIncludePaths.Add("Runtime/TraceLog/Public");
 
-		PrivateIncludePaths.AddRange(
-			new string[] {
-				"Runtime/SynthBenchmark/Public",
-				"Runtime/Engine/Public",
-			}
-			);
-
 		PrivateIncludePathModuleNames.AddRange(
 			new string[] {
-				"TargetPlatform",
+				"Analytics",
 				"DerivedDataCache",
+				"Engine",
 				"InputDevice",
-				"Analytics",
-				"RHI"
+				"RHI",
+				"SynthBenchmark",
+				"TargetPlatform",
 			}
 			);
 
diff --git a/Engine/Source/Runtime/GeometryCore/GeometryCore.Build.cs b/Engine/Source/Runtime/GeometryCore/GeometryCore.Build.cs
--- a/Engine/Source/Runtime/GeometryCore/GeometryCore.Build.cs
+++ b/Engine/Source/Runtime/GeometryCore/GeometryCore.Build.cs
@@ -8,12 +8,10 @@ public class GeometryCore : ModuleRules
 	{
         PCHUsage = ModuleRules.PCHUsageMode.UseExplicitOrSharedPCHs;
 
-		// For GPUSkinPublicDefs.h
-		PublicIncludePaths.Add("Runtime/Engine/Public");
-
 		PublicIncludePathModuleNames.AddRange(
 			new string[] {
 				"AnimationCore",			// For the BoneWeights.h include
+				"Engine",					// For GPUSkinPublicDefs.h
 			}
 		);
 
diff --git a/Engine/Source/Runtime/MRMesh/MRMesh.build.cs b/Engine/Source/Runtime/MRMesh/MRMesh.build.cs
--- a/Engine/Source/Runtime/MRMesh/MRMesh.build.cs
+++ b/Engine/Source/Runtime/MRMesh/MRMesh.build.cs
@@ -27,8 +27,7 @@ namespace UnrealBuildTool.Rules
 			}
 			
 			// Used for including the private Chaos headers
-			string EnginePath = Path.GetFullPath(Target.RelativeEnginePath);
-			PrivateIncludePaths.Add(Path.Combine(EnginePath, "Source/Runtime/Engine/Private"));
+			PrivateIncludePaths.Add(Path.Combine(GetModuleDirectory("Engine"), "Private"));
 
 			PrivateIncludePathModuleNames.Add("DerivedDataCache");
 		}
diff --git a/Engine/Source/Runtime/NavigationSystem/NavigationSystem.Build.cs b/Engine/Source/Runtime/NavigationSystem/NavigationSystem.Build.cs
--- a/Engine/Source/Runtime/NavigationSystem/NavigationSystem.Build.cs
+++ b/Engine/Source/Runtime/NavigationSystem/NavigationSystem.Build.cs
@@ -1,21 +1,19 @@
 // Copyright Epic Games, Inc. All Rights Reserved.
 
+using System.IO;
+
 namespace UnrealBuildTool.Rules
 {
     public class NavigationSystem : ModuleRules
     {
         public NavigationSystem(ReadOnlyTargetRules Target) : base(Target)
         {
-            PublicIncludePaths.AddRange(
-                new string[] {
-                    "Runtime/NavigationSystem/Public",
-                }
-                );
+			PublicIncludePathModuleNames.Add("NavigationSystem");
 
             PrivateIncludePaths.AddRange(
                 new string[] {
-                    "Runtime/NavigationSystem/Private",
-                    "Runtime/Engine/Private",
+					Path.Combine(GetModuleDirectory("NavigationSystem"), "Private"),
+					Path.Combine(GetModuleDirectory("Engine"), "Private"),
                 }
                 );
 
diff --git a/Engine/Source/Runtime/SkeletalMeshDescription/SkeletalMeshDescription.Build.cs b/Engine/Source/Runtime/SkeletalMeshDescription/SkeletalMeshDescription.Build.cs
--- a/Engine/Source/Runtime/SkeletalMeshDescription/SkeletalMeshDescription.Build.cs
+++ b/Engine/Source/Runtime/SkeletalMeshDescription/SkeletalMeshDescription.Build.cs
@@ -6,14 +6,11 @@ namespace UnrealBuildTool.Rules
 	{
 		public SkeletalMeshDescription(ReadOnlyTargetRules Target) : base(Target)
 		{
-            // For GPUSkinPublicDefs.h
-            PublicIncludePaths.Add("Runtime/Engine/Public");
-            
-            // For BoneWeights
 			PublicIncludePathModuleNames.AddRange(
 				new string[]
 				{
-					"AnimationCore"
+					"AnimationCore",	// For BoneWeights
+					"Engine",			// For GPUSkinPublicDefs.h
 				}
 			);
 
