PR #11095: Fix sampling UV on rectangular textures

#11095

diff --git a/Engine/Plugins/Editor/BlueprintMaterialTextureNodes/Source/BlueprintMaterialTextureNodes/Private/BlueprintMaterialTextureNodesBPLibrary.cpp b/Engine/Plugins/Editor/BlueprintMaterialTextureNodes/Source/BlueprintMaterialTextureNodes/Private/BlueprintMaterialTextureNodesBPLibrary.cpp
--- a/Engine/Plugins/Editor/BlueprintMaterialTextureNodes/Source/BlueprintMaterialTextureNodes/Private/BlueprintMaterialTextureNodesBPLibrary.cpp
+++ b/Engine/Plugins/Editor/BlueprintMaterialTextureNodes/Source/BlueprintMaterialTextureNodes/Private/BlueprintMaterialTextureNodesBPLibrary.cpp
@@ -46,7 +46,7 @@ FLinearColor UBlueprintMaterialTextureNodesBPLibrary::Texture2D_SampleUV_EditorO
 		int32 MipWidth = CurMip->SizeX;
 		int32 MipHeight = CurMip->SizeY;
 		int32 X = FMath::Clamp(int(UV.X * MipWidth), 0, MipWidth - 1);
-		int32 Y = FMath::Clamp(int(UV.Y * MipHeight), 0, MipWidth - 1);
+		int32 Y = FMath::Clamp(int(UV.Y * MipHeight), 0, MipHeight - 1);
 
 		//Retrieve from BulkData if available and if uncompressed format
 		if ( (ImageData->IsBulkDataLoaded() && ImageData->GetBulkDataSize() > 0) && (Texture->GetPixelFormat() == PF_B8G8R8A8 || Texture->GetPixelFormat() == PF_FloatRGBA) )
