BlueprintMaterialTextureNodes: Updated CreateMIC_EditorOnly to be consistent with the path handling from KismetRenderingLibrary, allowing creating assets in plugin folders. No longer forces path to begin with "/game/".

diff --git a/Engine/Plugins/Editor/BlueprintMaterialTextureNodes/Source/BlueprintMaterialTextureNodes/Private/BlueprintMaterialTextureNodesBPLibrary.cpp b/Engine/Plugins/Editor/BlueprintMaterialTextureNodes/Source/BlueprintMaterialTextureNodes/Private/BlueprintMaterialTextureNodesBPLibrary.cpp
--- a/Engine/Plugins/Editor/BlueprintMaterialTextureNodes/Source/BlueprintMaterialTextureNodes/Private/BlueprintMaterialTextureNodesBPLibrary.cpp
+++ b/Engine/Plugins/Editor/BlueprintMaterialTextureNodes/Source/BlueprintMaterialTextureNodes/Private/BlueprintMaterialTextureNodesBPLibrary.cpp
@@ -239,7 +239,7 @@ UMaterialInstanceConstant* UBlueprintMaterialTextureNodesBPLibrary::CreateMIC_Ed
 		IAssetTools& AssetTools = FModuleManager::Get().LoadModuleChecked<FAssetToolsModule>("AssetTools").Get();
 
 		//Use asset name only if directories are specified, otherwise full path
-		if (!InName.Contains(TEXT("/")))
+		if (!InName.StartsWith(TEXT("/")))
 		{
 			FString AssetName = Material->GetOutermost()->GetName();
 			const FString SanitizedBasePackageName = UPackageTools::SanitizePackageName(AssetName);
@@ -248,9 +248,6 @@ UMaterialInstanceConstant* UBlueprintMaterialTextureNodesBPLibrary::CreateMIC_Ed
 		}
 		else
 		{
-			InName.RemoveFromStart(TEXT("/"));
-			InName.RemoveFromStart(TEXT("Content/"));
-			InName.StartsWith(TEXT("Game/")) == true ? InName.InsertAt(0, TEXT("/")) : InName.InsertAt(0, TEXT("/Game/"));
 			AssetTools.CreateUniqueAssetName(InName, TEXT(""), PackageName, Name);
 		}
 
